Add missing icons for nDisplay apps, plugin, template

diff --git a/Engine/Source/Programs/nDisplayLauncher/nDisplayLauncher.csproj b/Engine/Source/Programs/nDisplayLauncher/nDisplayLauncher.csproj
--- a/Engine/Source/Programs/nDisplayLauncher/nDisplayLauncher.csproj
+++ b/Engine/Source/Programs/nDisplayLauncher/nDisplayLauncher.csproj
@@ -63,8 +63,7 @@
   </PropertyGroup>
   <PropertyGroup />
   <PropertyGroup>
-    <ApplicationIcon>
-    </ApplicationIcon>
+    <ApplicationIcon>Resources\nDisplayLauncher.ico</ApplicationIcon>
   </PropertyGroup>
   <PropertyGroup />
   <PropertyGroup />
@@ -202,7 +201,9 @@
       <Install>false</Install>
     </BootstrapperPackage>
   </ItemGroup>
-  <ItemGroup />
+  <ItemGroup>
+    <Resource Include="Resources\nDisplayLauncher.ico" />
+  </ItemGroup>
   <Import Project="$(MSBuildToolsPath)\Microsoft.CSharp.targets" />
   <!-- To modify your build process, add your task inside one of the targets below and uncomment it. 
        Other similar extension points exist, see Microsoft.Common.targets.
diff --git a/Engine/Source/Programs/nDisplayListener/nDisplayListener.csproj b/Engine/Source/Programs/nDisplayListener/nDisplayListener.csproj
--- a/Engine/Source/Programs/nDisplayListener/nDisplayListener.csproj
+++ b/Engine/Source/Programs/nDisplayListener/nDisplayListener.csproj
@@ -82,6 +82,9 @@
     <ErrorReport>prompt</ErrorReport>
     <CodeAnalysisRuleSet>MinimumRecommendedRules.ruleset</CodeAnalysisRuleSet>
   </PropertyGroup>
+  <PropertyGroup>
+    <ApplicationIcon>Resources\nDisplayListener.ico</ApplicationIcon>
+  </PropertyGroup>
   <ItemGroup>
     <Reference Include="System" />
     <Reference Include="System.Core" />
@@ -113,6 +116,9 @@
       <Install>false</Install>
     </BootstrapperPackage>
   </ItemGroup>
+  <ItemGroup>
+    <Content Include="Resources\nDisplayListener.ico" />
+  </ItemGroup>
   <Import Project="$(MSBuildToolsPath)\Microsoft.CSharp.targets" />
   <!-- To modify your build process, add your task inside one of the targets below and uncomment it. 
        Other similar extension points exist, see Microsoft.Common.targets.
