Fixed build warning: Could not find rule set file "MinimumRecommendedRules.ruleset"

#rnx

diff --git a/Engine/Source/Programs/nDisplayLauncher/nDisplayLauncher.csproj b/Engine/Source/Programs/nDisplayLauncher/nDisplayLauncher.csproj
--- a/Engine/Source/Programs/nDisplayLauncher/nDisplayLauncher.csproj
+++ b/Engine/Source/Programs/nDisplayLauncher/nDisplayLauncher.csproj
@@ -91,7 +91,7 @@
     <DebugType>pdbonly</DebugType>
     <PlatformTarget>AnyCPU</PlatformTarget>
     <ErrorReport>prompt</ErrorReport>
-    <CodeAnalysisRuleSet>MinimumRecommendedRules.ruleset</CodeAnalysisRuleSet>
+    <CodeAnalysisRuleSet></CodeAnalysisRuleSet>
     <Prefer32Bit>true</Prefer32Bit>
   </PropertyGroup>
   <ItemGroup>
diff --git a/Engine/Source/Programs/nDisplayListener/nDisplayListener.csproj b/Engine/Source/Programs/nDisplayListener/nDisplayListener.csproj
--- a/Engine/Source/Programs/nDisplayListener/nDisplayListener.csproj
+++ b/Engine/Source/Programs/nDisplayListener/nDisplayListener.csproj
@@ -80,7 +80,7 @@
     <DebugType>pdbonly</DebugType>
     <PlatformTarget>AnyCPU</PlatformTarget>
     <ErrorReport>prompt</ErrorReport>
-    <CodeAnalysisRuleSet>MinimumRecommendedRules.ruleset</CodeAnalysisRuleSet>
+    <CodeAnalysisRuleSet></CodeAnalysisRuleSet>
   </PropertyGroup>
   <PropertyGroup>
     <ApplicationIcon>Resources\nDisplayListener.ico</ApplicationIcon>
